Change UnrealHitchParser to target .NET Framework 4.5 to fix warning.

diff --git a/Engine/Source/Programs/UnrealHitchParser/App.config b/Engine/Source/Programs/UnrealHitchParser/App.config
--- a/Engine/Source/Programs/UnrealHitchParser/App.config
+++ b/Engine/Source/Programs/UnrealHitchParser/App.config
@@ -1,6 +1,6 @@
-﻿<?xml version="1.0" encoding="utf-8" ?>
+﻿<?xml version="1.0" encoding="utf-8"?>
 <configuration>
     <startup> 
-        <supportedRuntime version="v4.0" sku=".NETFramework,Version=v4.5.2" />
+        <supportedRuntime version="v4.0" sku=".NETFramework,Version=v4.5"/>
     </startup>
-</configuration>
\ No newline at end of file
+</configuration>
diff --git a/Engine/Source/Programs/UnrealHitchParser/UnrealHitchParser.csproj b/Engine/Source/Programs/UnrealHitchParser/UnrealHitchParser.csproj
--- a/Engine/Source/Programs/UnrealHitchParser/UnrealHitchParser.csproj
+++ b/Engine/Source/Programs/UnrealHitchParser/UnrealHitchParser.csproj
@@ -9,9 +9,10 @@
     <AppDesignerFolder>Properties</AppDesignerFolder>
     <RootNamespace>UnrealHitchParser</RootNamespace>
     <AssemblyName>UnrealHitchParser</AssemblyName>
-    <TargetFrameworkVersion>v4.5.2</TargetFrameworkVersion>
+    <TargetFrameworkVersion>v4.5</TargetFrameworkVersion>
     <FileAlignment>512</FileAlignment>
     <AutoGenerateBindingRedirects>true</AutoGenerateBindingRedirects>
+    <TargetFrameworkProfile />
   </PropertyGroup>
   <PropertyGroup Condition=" '$(Configuration)|$(Platform)' == 'Debug|AnyCPU' ">
     <PlatformTarget>AnyCPU</PlatformTarget>
